[UnrealVCam] [UnrealStage] [FlutterCommon] Add third-party notices menu

diff --git a/Engine/Extras/Flutter/FlutterCommon/tools/generate_license_assets.py b/Engine/Extras/Flutter/FlutterCommon/tools/generate_license_assets.py
new file mode 100644
--- /dev/null
+++ b/Engine/Extras/Flutter/FlutterCommon/tools/generate_license_assets.py
@@ -0,0 +1,107 @@
+# Copyright Epic Games, Inc. All Rights Reserved.
+
+import argparse
+import hashlib
+import json
+import os
+import re
+import stat
+import sys
+
+parser = argparse.ArgumentParser(
+    description='Takes a directory containing license files, and produces a manifest with associated assets that can be loaded at runtime in Flutter.',
+    formatter_class=argparse.ArgumentDefaultsHelpFormatter,
+)
+parser.add_argument('-i', '--input', default='licenses', type=str, metavar='INPUT_PATH', help='The path containing directories with license files')
+parser.add_argument('-o', '--output', default='assets/licenses', type=str, metavar='OUTPUT_PATH', help='The path to which to write the manifest and text assets')
+
+args = parser.parse_args(sys.argv[1:])
+
+notices_directory = args.input
+output_directory = args.output
+text_directory = os.path.join(output_directory, 'text')
+manifest_path = os.path.join(output_directory, 'manifest.json')
+
+if not os.path.exists(notices_directory):
+    print(f'Invalid input path "{notices_directory}"')
+    sys.exit(1)
+
+# Check that we're not about to overwrite something unrelated
+if os.path.exists(text_directory) and not os.path.exists(manifest_path):
+    response = input(
+        f'Output path "{output_directory}" exists, but does not contain a manifest file.\n'
+        f'All files in the the path "{text_directory}" will be deleted and re-generated.\n'
+        'Are you sure you want to continue? (Y/N) ')
+    
+    if response.lower() != 'y':
+        print('Aborting')
+        sys.exit(1)
+
+os.makedirs(output_directory, exist_ok=True)
+os.makedirs(text_directory, exist_ok=True)
+
+license_pattern = re.compile(
+    r'=*\n'
+    r'Software Name: ([^\n]*)\n'
+    r'Version: ([^\n]*)\n'
+    r'URL: .*\n'
+    r'===+\n'
+    r'(.*)',
+
+    flags=re.DOTALL
+)
+
+# Clear existing license text in the output directory since we intend to replace them all
+for license_name in os.listdir(text_directory):
+    license_path = os.path.join(text_directory, license_name)
+    os.chmod(license_path, stat.S_IWRITE) # Clear read-only flag so we can delete files tracked by Perforce
+    os.remove(license_path)
+
+manifest = {}
+seen_hashes = set()
+
+for input_directory_item in sorted(os.listdir(notices_directory), key=str.lower):
+    # Find the license file
+    input_directory_path = os.path.join(notices_directory, input_directory_item)
+
+    # Skip any paths that don't point to a directory
+    if os.path.isdir(input_directory_path) is False:
+        continue 
+
+    license_path = None
+
+    for file_name in os.listdir(input_directory_path):
+        if file_name.lower().endswith('.license'):
+            license_path = os.path.join(input_directory_path, file_name)
+            break
+    
+    if not license_path:
+        print(f'No license found for directory {input_directory_path}')
+        sys.exit(1)
+    
+    # Parse the license file
+    with open(license_path, 'r', encoding='utf-8') as license_file:
+        match = license_pattern.search(license_file.read())
+        if not match:
+            print(f'Invalid license format for file {license_path}')
+            sys.exit(1)
+        
+        name, version, license_text = match.groups()
+
+    # Make a hash for each license so we only store one copy of each
+    license_hash = hashlib.md5(license_text.encode()).hexdigest()
+
+    if license_hash not in seen_hashes:
+        with open(os.path.join(text_directory, license_hash), 'w', encoding='utf-8') as license_text_file:
+            license_text_file.write(license_text)
+
+    manifest[name] = {
+        'version': version,
+        'file': license_hash
+    }
+
+# Write manifest
+with open(manifest_path, 'w') as manifest_file:
+    manifest_file.write(json.dumps(manifest))
+
+print(f'Asset generation complete. If this is a Perforce repo, you should reconcile the "{output_directory}" directory.')
\ No newline at end of file
diff --git a/Engine/Extras/VirtualProduction/EpicStageApp/lib/main.dart b/Engine/Extras/VirtualProduction/EpicStageApp/lib/main.dart
--- a/Engine/Extras/VirtualProduction/EpicStageApp/lib/main.dart
+++ b/Engine/Extras/VirtualProduction/EpicStageApp/lib/main.dart
@@ -7,6 +7,7 @@ import 'package:epic_common/localizations.dart';
 import 'package:epic_common/logging.dart';
 import 'package:epic_common/preferences.dart';
 import 'package:epic_common/theme.dart';
+import 'package:epic_common/third_party_notice.dart';
 import 'package:epic_common/utilities/version.dart';
 import 'package:epic_common/widgets.dart';
 import 'package:flutter/material.dart';
@@ -165,6 +166,9 @@ class _EpicStageAppState extends State<EpicStageApp> with TickerProviderStateMix
 
     return MultiProvider(
       providers: [
+        Provider<ThirdPartyNoticeManifest>(
+          create: (context) => ThirdPartyNoticeManifest(assetBundle: DefaultAssetBundle.of(context)),
+        ),
         Provider<EnginePassphraseManager>(create: (_) => EnginePassphraseManager()),
         Provider<PreferencesBundle>(create: (_) => _preferenceBundle),
         Provider<ConnectionSettings>(create: (_) => ConnectionSettings(_preferenceBundle)),
diff --git a/Engine/Extras/VirtualProduction/EpicStageApp/lib/widgets/screens/main/toolbar/settings/pages/settings_dialog_main.dart b/Engine/Extras/VirtualProduction/EpicStageApp/lib/widgets/screens/main/toolbar/settings/pages/settings_dialog_main.dart
--- a/Engine/Extras/VirtualProduction/EpicStageApp/lib/widgets/screens/main/toolbar/settings/pages/settings_dialog_main.dart
+++ b/Engine/Extras/VirtualProduction/EpicStageApp/lib/widgets/screens/main/toolbar/settings/pages/settings_dialog_main.dart
@@ -1,5 +1,6 @@
 // Copyright Epic Games, Inc. All Rights Reserved.
 
+import 'package:epic_common/localizations.dart';
 import 'package:epic_common/theme.dart';
 import 'package:epic_common/utilities/version.dart';
 import 'package:epic_common/widgets.dart';
@@ -16,6 +17,7 @@ import '../../../../eula/eula_screen.dart';
 import 'settings_advanced_view.dart';
 import 'settings_log_list.dart';
 import 'settings_root_actor_picker.dart';
+import 'settings_third_party_notice_list.dart';
 
 /// Main screen of the [SettingsDialog].
 class SettingsDialogMain extends StatelessWidget {
@@ -81,6 +83,11 @@ class SettingsDialogMain extends StatelessWidget {
             iconPath: 'packages/epic_common/assets/icons/info.svg',
             onTap: () => launchUrl(Uri.parse(_privacyPolicyUri)),
           ),
+          SettingsMenuItem(
+            title: EpicCommonLocalizations.of(context)!.settingsThirdPartyNoticeTitle,
+            iconPath: 'packages/epic_common/assets/icons/info.svg',
+            onTap: () => Navigator.of(context).pushNamed(SettingsThirdPartyNoticeList.route),
+          ),
           const SettingsMenuDivider(),
           Container(
             padding: EdgeInsets.symmetric(horizontal: 8),
diff --git a/Engine/Extras/VirtualProduction/EpicStageApp/lib/widgets/screens/main/toolbar/settings/pages/settings_third_party_notice_list.dart b/Engine/Extras/VirtualProduction/EpicStageApp/lib/widgets/screens/main/toolbar/settings/pages/settings_third_party_notice_list.dart
new file mode 100644
--- /dev/null
+++ b/Engine/Extras/VirtualProduction/EpicStageApp/lib/widgets/screens/main/toolbar/settings/pages/settings_third_party_notice_list.dart
@@ -0,0 +1,19 @@
+// Copyright Epic Games, Inc. All Rights Reserved.
+
+import 'package:epic_common/widgets.dart';
+import 'package:flutter/material.dart';
+
+import 'settings_third_party_notice_view.dart';
+
+class SettingsThirdPartyNoticeList extends StatelessWidget {
+  const SettingsThirdPartyNoticeList({super.key});
+
+  static const String route = '/third_party';
+
+  @override
+  Widget build(BuildContext context) {
+    return ThirdPartyNoticeListSettingsPage(
+      viewRoute: SettingsThirdPartyNoticeView.route,
+    );
+  }
+}
diff --git a/Engine/Extras/VirtualProduction/EpicStageApp/lib/widgets/screens/main/toolbar/settings/pages/settings_third_party_notice_view.dart b/Engine/Extras/VirtualProduction/EpicStageApp/lib/widgets/screens/main/toolbar/settings/pages/settings_third_party_notice_view.dart
new file mode 100644
--- /dev/null
+++ b/Engine/Extras/VirtualProduction/EpicStageApp/lib/widgets/screens/main/toolbar/settings/pages/settings_third_party_notice_view.dart
@@ -0,0 +1,15 @@
+// Copyright Epic Games, Inc. All Rights Reserved.
+
+import 'package:epic_common/widgets.dart';
+import 'package:flutter/material.dart';
+
+class SettingsThirdPartyNoticeView extends StatelessWidget {
+  const SettingsThirdPartyNoticeView({super.key});
+
+  static const String route = '/third_party/view';
+
+  @override
+  Widget build(BuildContext context) {
+    return ThirdPartyNoticeViewSettingsPage();
+  }
+}
diff --git a/Engine/Extras/VirtualProduction/EpicStageApp/lib/widgets/screens/main/toolbar/settings/settings_dialog.dart b/Engine/Extras/VirtualProduction/EpicStageApp/lib/widgets/screens/main/toolbar/settings/settings_dialog.dart
--- a/Engine/Extras/VirtualProduction/EpicStageApp/lib/widgets/screens/main/toolbar/settings/settings_dialog.dart
+++ b/Engine/Extras/VirtualProduction/EpicStageApp/lib/widgets/screens/main/toolbar/settings/settings_dialog.dart
@@ -10,6 +10,8 @@ import 'pages/settings_dialog_main.dart';
 import 'pages/settings_log_list.dart';
 import 'pages/settings_log_view.dart';
 import 'pages/settings_root_actor_picker.dart';
+import 'pages/settings_third_party_notice_list.dart';
+import 'pages/settings_third_party_notice_view.dart';
 
 /// Dialog containing the app's settings menu.
 class SettingsDialog extends StatefulWidget {
@@ -85,6 +87,14 @@ class _SettingsDialogState extends State<SettingsDialog> with RouteAware {
                   page = const SettingsCompressionPicker();
                   break;
 
+                case SettingsThirdPartyNoticeList.route:
+                  page = const SettingsThirdPartyNoticeList();
+                  break;
+
+                case SettingsThirdPartyNoticeView.route:
+                  page = const SettingsThirdPartyNoticeView();
+                  break;
+
                 default:
                   throw Exception('No settings route named ${settings.name}');
               }
diff --git a/Engine/Extras/VirtualProduction/EpicStageApp/pubspec.yaml b/Engine/Extras/VirtualProduction/EpicStageApp/pubspec.yaml
--- a/Engine/Extras/VirtualProduction/EpicStageApp/pubspec.yaml
+++ b/Engine/Extras/VirtualProduction/EpicStageApp/pubspec.yaml
@@ -85,6 +85,8 @@ flutter:
   assets:
     - assets/images/map/
     - assets/images/icons/
+    - assets/licenses/
+    - assets/licenses/text/
 
   fonts:
     - family: Inter
diff --git a/Engine/Extras/VirtualProduction/FlutterLiveLinkVCAM/pubspec.yaml b/Engine/Extras/VirtualProduction/FlutterLiveLinkVCAM/pubspec.yaml
--- a/Engine/Extras/VirtualProduction/FlutterLiveLinkVCAM/pubspec.yaml
+++ b/Engine/Extras/VirtualProduction/FlutterLiveLinkVCAM/pubspec.yaml
@@ -42,6 +42,7 @@ dev_dependencies:
     sdk: flutter
   flutter_lints: ^4.0.0
 
+  arb_generator: 0.0.4
   pigeon: ^22.6.3
   flutter_launcher_icons: ^0.14.1
 
@@ -58,6 +59,8 @@ flutter:
   assets:
     - assets/eula.html
     - assets/icons/
+    - assets/licenses/
+    - assets/licenses/text/
 
   fonts:
     - family: Inter
